Added definition of FDebuggingVisualizerState to LiveCoding module globals

This ensures that natvis expressions are still able to use this symbol's members even from code that has been repatched with Live Coding.

# robomerge 5.5

diff --git a/Engine/Extras/NatvisHelpers/Globals.h b/Engine/Extras/NatvisHelpers/Globals.h
--- a/Engine/Extras/NatvisHelpers/Globals.h
+++ b/Engine/Extras/NatvisHelpers/Globals.h
@@ -8,5 +8,16 @@ class FChunkedFixedUObjectArray;
 extern uint8** GNameBlocksDebug;
 extern FChunkedFixedUObjectArray* GObjectArrayForDebugVisualizers;
 
-namespace UE { namespace Core { struct FVisualizerDebuggingState; } }
+namespace UE
+{
+	namespace Core {
+		struct FVisualizerDebuggingStateImpl;
+		struct FVisualizerDebuggingState
+		{
+			const char* GuidString;
+			void** Ptrs;
+			FVisualizerDebuggingStateImpl* PimplData;
+		};
+	}
+}
 extern UE::Core::FVisualizerDebuggingState* GDebuggingState;
